OodleDataCompressionFormat remove stats timers
we already have redundant stats timers in compression.cpp

diff --git a/Engine/Plugins/Compression/OodleData/Source/Private/OodleDataCompressionFormat.cpp b/Engine/Plugins/Compression/OodleData/Source/Private/OodleDataCompressionFormat.cpp
--- a/Engine/Plugins/Compression/OodleData/Source/Private/OodleDataCompressionFormat.cpp
+++ b/Engine/Plugins/Compression/OodleData/Source/Private/OodleDataCompressionFormat.cpp
@@ -62,8 +62,6 @@ PackageCompressionLevel_DebugDevelopment=3
 #include "Misc/ICompressionFormat.h"
 #include "Misc/CommandLine.h"
 #include "Misc/Parse.h"
-#include "Stats/Stats.h"
-#include "Stats/Stats2.h"
 
 #if WITH_EDITOR
 #include "Settings/ProjectPackagingSettings.h"
@@ -89,12 +87,6 @@ DEFINE_LOG_CATEGORY_STATIC(OodleDataCompression, Log, All);
 //#define MAX_OODLE_DECODER_SIZE_ON_STACK	 100000
 
 
-#if STATS
-DECLARE_STATS_GROUP( TEXT( "Compression" ), STATGROUP_Compression, STATCAT_Advanced );
-DECLARE_FLOAT_ACCUMULATOR_STAT(TEXT("Total oodle encode time"), STAT_Compression_Oodle_Encode, STATGROUP_Compression);
-DECLARE_FLOAT_ACCUMULATOR_STAT(TEXT("Total oodle decode time"), STAT_Compression_Oodle_Decode, STATGROUP_Compression);
-#endif
-
 
 struct FOodleDataCompressionFormat : ICompressionFormat
 {
@@ -283,11 +275,6 @@ struct FOodleDataCompressionFormat : ICompressionFormat
 
 	virtual bool Compress(void* OutCompressedBuffer, int32& OutCompressedSize, const void* InUncompressedBuffer, int32 InUncompressedSize, int32 InCompressionData) override
 	{
-#if STATS
-		DECLARE_SCOPE_CYCLE_COUNTER(TEXT("Compress Memory Oodle"), STAT_Compression_Oodle_Encode_Cycles, STATGROUP_Compression);
-		SCOPE_SECONDS_ACCUMULATOR(STAT_Compression_Oodle_Encode);
-#endif
-
 		// OutCompressedSize is read-write
 		int32 CompressedBufferSize = OutCompressedSize;
 
@@ -313,10 +300,6 @@ struct FOodleDataCompressionFormat : ICompressionFormat
 
 	virtual bool Uncompress(void* OutUncompressedBuffer, int32& OutUncompressedSize, const void* InCompressedBuffer, int32 InCompressedSize, int32 CompressionData) override
 	{
-#if STATS
-		DECLARE_SCOPE_CYCLE_COUNTER(TEXT("Uncompress Memory Oodle"), STAT_Compression_Oodle_Decode_Cycles, STATGROUP_Compression);
-		SCOPE_SECONDS_ACCUMULATOR(STAT_Compression_Oodle_Decode);
-#endif
 		// OutUncompressedSize is read-write
 		int32 UncompressedSize = OutUncompressedSize;
 		int Result = OodleDecode(InCompressedBuffer, InCompressedSize, OutUncompressedBuffer, UncompressedSize);
