OodleDataCompressionFormat : do command line processing in shipping with editor

diff --git a/Engine/Plugins/Compression/OodleData/Source/Private/OodleDataCompressionFormat.cpp b/Engine/Plugins/Compression/OodleData/Source/Private/OodleDataCompressionFormat.cpp
--- a/Engine/Plugins/Compression/OodleData/Source/Private/OodleDataCompressionFormat.cpp
+++ b/Engine/Plugins/Compression/OodleData/Source/Private/OodleDataCompressionFormat.cpp
@@ -329,11 +329,11 @@ class FOodleDataCompressionFormatModuleInterface : public IModuleInterface
 	{
 		// settings to use in non-tools context (eg. runtime game encoding) :
 		// (SetDefaultOodleOptionsForPackaging sets options for pak compression & iostore)
-		OodleLZ_Compressor UsedCompressor = OodleLZ_Compressor_Kraken;
+		OodleLZ_Compressor UsedCompressor = OodleLZ_Compressor_Mermaid;
 		OodleLZ_CompressionLevel UsedLevel = OodleLZ_CompressionLevel_Fast;
 		int32 SpaceSpeedTradeoff = 0;
 
-		#if ! UE_BUILD_SHIPPING
+		#if ( (!UE_BUILD_SHIPPING) || WITH_EDITOR )
 		{
 
 		// parse the command line to get compressor & level settings
