change PackageCompressionLevel_DebugDevelopment to 4 instead of 3

diff --git a/Engine/Config/BaseGame.ini b/Engine/Config/BaseGame.ini
--- a/Engine/Config/BaseGame.ini
+++ b/Engine/Config/BaseGame.ini
@@ -92,7 +92,7 @@ PackageAdditionalCompressionOptions=-compressionblocksize=256KB
 PackageCompressionMethod=Kraken
 PackageCompressionLevel_Distribution=7
 PackageCompressionLevel_TestShipping=5
-PackageCompressionLevel_DebugDevelopment=3
+PackageCompressionLevel_DebugDevelopment=4
 PackageCompressionMinBytesSaved=1024
 PackageCompressionMinPercentSaved=5
 bForceUseProjectCompressionFormatIgnoreHardwareOverride=False
diff --git a/Engine/Plugins/Compression/OodleData/Source/Private/OodleDataCompressionFormat.cpp b/Engine/Plugins/Compression/OodleData/Source/Private/OodleDataCompressionFormat.cpp
--- a/Engine/Plugins/Compression/OodleData/Source/Private/OodleDataCompressionFormat.cpp
+++ b/Engine/Plugins/Compression/OodleData/Source/Private/OodleDataCompressionFormat.cpp
@@ -32,11 +32,11 @@
 bCompressed=True
 bForceUseProjectCompressionFormat=False
 PackageCompressionFormat=Oodle
-PackageAdditionalCompressionOptions=-compressionblocksize=1MB -asynccompression
-PackageCompressionMethod=Mermaid
-PackageCompressionLevel_Distribution=8
+PackageAdditionalCompressionOptions=-compressionblocksize=256KB
+PackageCompressionMethod=Kraken
+PackageCompressionLevel_Distribution=7
 PackageCompressionLevel_TestShipping=5
-PackageCompressionLevel_DebugDevelopment=3
+PackageCompressionLevel_DebugDevelopment=4
 
 * This can be set in DefaultGame.ini then overrides set up per-platform.
 * 
