Oodle: The number of bufers kept around for data decoding is now configurable in Engine.ini. Additionally, buffers are allocated on an as-needed basis to avoid memory costs if Oodle is not used. #rb charles.bloom #pf 60c3bec0e5979e0001e693d4

diff --git a/Engine/Config/BaseEngine.ini b/Engine/Config/BaseEngine.ini
--- a/Engine/Config/BaseEngine.ini
+++ b/Engine/Config/BaseEngine.ini
@@ -2821,3 +2821,7 @@ LogVerbosity=0
 GlobalLambdaMultiplier=1.000000
 DefaultRDOLambda=30
 
+[OodleDataCompressionFormat]
+; Specifies how many buffers Oodle should keep around for 
+; reuse as opposed to heap allocating every decode.
+PreallocatedBufferCount=2
\ No newline at end of file
diff --git a/Engine/Plugins/Compression/OodleData/Source/Private/OodleDataCompressionFormat.cpp b/Engine/Plugins/Compression/OodleData/Source/Private/OodleDataCompressionFormat.cpp
--- a/Engine/Plugins/Compression/OodleData/Source/Private/OodleDataCompressionFormat.cpp
+++ b/Engine/Plugins/Compression/OodleData/Source/Private/OodleDataCompressionFormat.cpp
@@ -61,6 +61,7 @@ PackageCompressionLevel_DebugDevelopment=4
 #include "Misc/Compression.h"
 #include "Misc/ICompressionFormat.h"
 #include "Misc/CommandLine.h"
+#include "Misc/ConfigCacheIni.h"
 #include "Misc/Parse.h"
 
 #if WITH_EDITOR
@@ -73,11 +74,6 @@ DEFINE_LOG_CATEGORY_STATIC(OodleDataCompression, Log, All);
 
 #define OODLE_DERIVEDDATA_VER TEXT("BA7AA26CD1C3498787A3F3AA53895042")
 
-// Pre-allocates this many decode temporary buffers. 
-// More means less dynamic allocation, but more static memory overhead. 
-// optimal number may vary depending on Platform,OS,etc...
-#define NUM_OODLE_DECODE_BUFFERS 2
-
 // if the Decoder object is <= this size, just put it on the stack
 // MAX_OODLE_DECODER_SIZE_ON_STACK needs to be at least 64k to ever be used
 // the benefit of this could be more parallel decodes if the 2 pre-allocated buffers are in use
@@ -94,9 +90,10 @@ struct FOodleDataCompressionFormat : ICompressionFormat
 	OodleLZ_CompressionLevel CompressionLevel;
 	OodleLZ_CompressOptions CompressionOptions;
 	int32 OodleDecoderMemorySize;
+	int32 OodleDecoderBufferCount =0;
 
-	FCriticalSection OodleDecoderMutex[NUM_OODLE_DECODE_BUFFERS];
-	void * OodleDecoderMemory[NUM_OODLE_DECODE_BUFFERS];
+	FCriticalSection* OodleDecoderMutexes =0;
+	void ** OodleDecoderMemory =0;
 
 	FOodleDataCompressionFormat(OodleLZ_Compressor InCompressor, OodleLZ_CompressionLevel InCompressionLevel, int InSpaceSpeedTradeoffBytes)
 	{
@@ -109,35 +106,44 @@ struct FOodleDataCompressionFormat : ICompressionFormat
 
 		// enough decoder scratch for any compressor & buffer size.
 		// note "InCompressor" is what we want to Encode with but we may be asked to decode other compressors!
-		OodleDecoderMemorySize = (int32) OodleLZDecoder_MemorySizeNeeded(OodleLZ_Compressor_Invalid, -1);  
+		OodleDecoderMemorySize = (int32) OodleLZDecoder_MemorySizeNeeded(OodleLZ_Compressor_Invalid, -1);
 
-		for (int i = 0; i < NUM_OODLE_DECODE_BUFFERS; ++i)
+		int32 BufferCount = 2;
+		GConfig->GetInt(TEXT("OodleDataCompressionFormat"), TEXT("PreallocatedBufferCount"), BufferCount, GEngineIni);
+		if (BufferCount < 0)
 		{
-			OodleDecoderMemory[i] = FMemory::Malloc(OodleDecoderMemorySize);
-			if (OodleDecoderMemory[i] == NULL) 
-			{
-				UE_LOG(OodleDataCompression, Error, TEXT("FOodleDataCompressionFormat - Failed to allocate %d!"), OodleDecoderMemorySize);
-			}
+			BufferCount = 0;
+		}
+
+		OodleDecoderBufferCount = BufferCount;
+		if (OodleDecoderBufferCount)
+		{
+			OodleDecoderMutexes = new FCriticalSection[OodleDecoderBufferCount];
+			OodleDecoderMemory = new void* [OodleDecoderBufferCount];
+			FMemory::Memzero(OodleDecoderMemory, sizeof(void*) * OodleDecoderBufferCount);
 		}
 	}
 
 	virtual ~FOodleDataCompressionFormat() 
-	{
-		for (int i = 0; i < NUM_OODLE_DECODE_BUFFERS; ++i)
+	{		
+		if (OodleDecoderBufferCount)
 		{
-			if (OodleDecoderMutex[i].TryLock())
+			for (int i = 0; i < OodleDecoderBufferCount; ++i)
 			{
-				if (OodleDecoderMemory[i])
+				if (OodleDecoderMutexes[i].TryLock())
 				{
 					FMemory::Free(OodleDecoderMemory[i]);
 					OodleDecoderMemory[i] = 0;
+					OodleDecoderMutexes[i].Unlock();
+				}
+				else
+				{
+					UE_LOG(OodleDataCompression, Error, TEXT("FOodleDataCompressionFormat - shutting down while in use?"));
 				}
-				OodleDecoderMutex[i].Unlock();
-			}
-			else
-			{
-				UE_LOG(OodleDataCompression, Error, TEXT("FOodleDataCompressionFormat - shutting down while in use?"));
 			}
+
+			delete [] OodleDecoderMutexes;
+			delete [] OodleDecoderMemory;
 		}
 	}
 
@@ -219,10 +225,20 @@ struct FOodleDataCompressionFormat : ICompressionFormat
 		{
 
 			// try to take a mutex for one of the pre-allocated decode buffers
-			for (int i = 0; i < NUM_OODLE_DECODE_BUFFERS; ++i) 
+			for (int i = 0; i < OodleDecoderBufferCount; ++i)
 			{
-				if (OodleDecoderMutex[i].TryLock()) 
+				if (OodleDecoderMutexes[i].TryLock()) 
 				{
+					if (OodleDecoderMemory[i] == nullptr)
+					{
+						// Haven't allocated yet (we allocate on demand)
+						OodleDecoderMemory[i] = FMemory::Malloc(OodleDecoderMemorySize);
+						if (OodleDecoderMemory[i] == nullptr)
+						{
+							UE_LOG(OodleDataCompression, Error, TEXT("FOodleDataCompressionFormat - Failed to allocate preallocated	buffer %d bytes!"), OodleDecoderMemorySize);
+						}
+					}
+
 					if (OodleDecoderMemory[i])
 					{
 						//UE_LOG(OodleDataCompression, Display, TEXT("Decode with lock : %d -> %d"),InCompBufSize,InRawLen );
@@ -232,12 +248,12 @@ struct FOodleDataCompressionFormat : ICompressionFormat
 							NULL, 0, NULL, NULL,
 							OodleDecoderMemory[i], OodleDecoderMemorySize);
 
-						OodleDecoderMutex[i].Unlock();
+						OodleDecoderMutexes[i].Unlock();
 
 						return Result;
 					}
 
-					OodleDecoderMutex[i].Unlock();
+					OodleDecoderMutexes[i].Unlock();
 				}
 			}
 			
